Chaos Spatial Partitions: Fixing some bughawk issues.
- Fixed some minor test issues with not checking indices for being valid, etc...
- Fixed some static aabb tree issues with inserting then updating/removing an object without ever rebuilding the tree
- Added some guards for edge case in the static aabb tree time slicer config
#nojira

#Synced-CL 52508799

diff --git a/Engine/Source/Programs/ChaosSpatialPartitionsTests/Private/AabbTreeTestHelpers.h b/Engine/Source/Programs/ChaosSpatialPartitionsTests/Private/AabbTreeTestHelpers.h
--- a/Engine/Source/Programs/ChaosSpatialPartitionsTests/Private/AabbTreeTestHelpers.h
+++ b/Engine/Source/Programs/ChaosSpatialPartitionsTests/Private/AabbTreeTestHelpers.h
@@ -88,6 +88,11 @@ namespace Chaos::SpatialPartition
 
 	inline int32 BuildSimpleTree(TArray<FAabbTreeNode>& Nodes)
 	{
+		if (Nodes.Num() == 0)
+		{
+			return INDEX_NONE;
+		}
+
 		// Build the initial list of nodes to start grouping (the leaves)
 		TArray<int32> NodeIndices;
 		for (int32 I = 0; I < Nodes.Num(); ++I)
diff --git a/Engine/Source/Programs/ChaosSpatialPartitionsTests/Private/Library/StaticAabbTreeTests.cpp b/Engine/Source/Programs/ChaosSpatialPartitionsTests/Private/Library/StaticAabbTreeTests.cpp
--- a/Engine/Source/Programs/ChaosSpatialPartitionsTests/Private/Library/StaticAabbTreeTests.cpp
+++ b/Engine/Source/Programs/ChaosSpatialPartitionsTests/Private/Library/StaticAabbTreeTests.cpp
@@ -445,6 +445,12 @@ namespace Chaos::SpatialPartition
 			CHECK_THAT(Nodes[0]->Aabb, Catch::Equal(Objects[0].Aabb));
 			CHECK(Leaves[0]->NodeIndex == ActualRootIndex);
 		}
+		SECTION("Remove after insert deferred without rebuild")
+		{
+			FObjectData Object4{ .Aabb = BuildAabbCenterExtents(FVec3(10), AabbSize), .UserData = 3 };
+			SpatialPartition.InsertDeferred(Object4);
+			SpatialPartition.RemoveMinimal(Object4);
+		}
 		SECTION("Basic query validation")
 		{
 			// Just do a basic overlap query to ensure that we don't return the removed objects somehow
@@ -460,6 +466,40 @@ namespace Chaos::SpatialPartition
 		}
 	}
 
+	TEST_CASE("StaticAabbTree - UpdateIfWithinLeaf", "[Chaos][StaticAabbTree][spatial-partition]")
+	{
+		const FVec3 AabbSize(1);
+		TArray<FObjectData> Objects
+		{
+			FObjectData{.Aabb = BuildAabbCenterExtents(FVec3(0), AabbSize), .UserData = 0},
+			FObjectData{.Aabb = BuildAabbCenterExtents(FVec3(4), AabbSize), .UserData = 1},
+		};
+		const FAABB3 RootAabb = Union(Objects[0].Aabb, Objects[1].Aabb);
+
+		FObjectData& TestObject = Objects[0];
+		FTestStaticAabbTree SpatialPartition(FTestStaticAabbTree::FConfig{ .MaxElementsPerLeaf = 1 });
+		BuildFromObjects(SpatialPartition, Objects);
+
+		SECTION("New Aabb Within Leaf")
+		{
+			FAABB3 NewAabb = TestObject.Aabb;
+			NewAabb.ShrinkSymmetrically(FVec3(0.1));
+			CHECK(SpatialPartition.UpdateIfWithinLeaf(TestObject.UserData, NewAabb, TestObject.Handle));
+		}
+		SECTION("New Aabb Outside Leaf")
+		{
+			FAABB3 NewAabb = TestObject.Aabb;
+			NewAabb.MoveByVector(FVec3(1));
+			CHECK(!SpatialPartition.UpdateIfWithinLeaf(TestObject.UserData, NewAabb, TestObject.Handle));
+		}
+		SECTION("Object deferred inserted without rebuild")
+		{
+			FObjectData Object2{ .Aabb = BuildAabbCenterExtents(FVec3(10), AabbSize), .UserData = 2 };
+			SpatialPartition.InsertDeferred(Object2);
+			CHECK(!SpatialPartition.UpdateIfWithinLeaf(Object2.UserData, Object2.Aabb, Object2.Handle));
+		}
+	}
+
 	TEST_CASE("StaticAabbTree - Mutations During Rebuild", "[Chaos][StaticAabbTree][spatial-partition]")
 	{
 		const FTestStaticAabbTree::FConfig Config
diff --git a/Engine/Source/Programs/ChaosSpatialPartitionsTests/Private/StaticAabbTreeHelpers.h b/Engine/Source/Programs/ChaosSpatialPartitionsTests/Private/StaticAabbTreeHelpers.h
--- a/Engine/Source/Programs/ChaosSpatialPartitionsTests/Private/StaticAabbTreeHelpers.h
+++ b/Engine/Source/Programs/ChaosSpatialPartitionsTests/Private/StaticAabbTreeHelpers.h
@@ -51,7 +51,8 @@ namespace Chaos::SpatialPartition
 		for (int32 I = 0; I < IndexCount; ++I)
 		{
 			OutLeaves[I] = FindLeafPointerWithHandleValue(Leaves, I);
-			OutNodes[I] = (OutLeaves[I] == nullptr) ? nullptr : &Nodes[OutLeaves[I]->NodeIndex];
+			const int32 NodeIndex = (OutLeaves[I] != nullptr) ? OutLeaves[I]->NodeIndex : INDEX_NONE;
+			OutNodes[I] = Nodes.IsValidIndex(NodeIndex) ? &Nodes[NodeIndex] : nullptr;
 		}
 	}
 } // namespace Chaos::SpatialPartition
diff --git a/Engine/Source/Runtime/Experimental/ChaosSpatialPartitions/Private/ChaosSpatialPartitions/Algorithms/StaticAabbTreeTimeSlicer.cpp b/Engine/Source/Runtime/Experimental/ChaosSpatialPartitions/Private/ChaosSpatialPartitions/Algorithms/StaticAabbTreeTimeSlicer.cpp
--- a/Engine/Source/Runtime/Experimental/ChaosSpatialPartitions/Private/ChaosSpatialPartitions/Algorithms/StaticAabbTreeTimeSlicer.cpp
+++ b/Engine/Source/Runtime/Experimental/ChaosSpatialPartitions/Private/ChaosSpatialPartitions/Algorithms/StaticAabbTreeTimeSlicer.cpp
@@ -16,6 +16,14 @@ namespace Chaos::SpatialPartition::AabbTreeAlgorithm
 		{
 			return;
 		}
+		if (!ensure(Config.MaxElementsPerLeaf != 0))
+		{
+			Config.MaxElementsPerLeaf = 1;
+		}
+		if (!ensure(Config.SurfaceAreaHeuristicBinCount > 0))
+		{
+			Config.SurfaceAreaHeuristicBinCount = 1;
+		}
 
 		const int32 MinimumLeafCount = (ElementCount / Config.MaxElementsPerLeaf) + 1;
 		const int32 MinimumNodeCount = MinimumLeafCount * 2 - 1;
diff --git a/Engine/Source/Runtime/Experimental/ChaosSpatialPartitions/Private/ChaosSpatialPartitions/Collections/DefaultSpatialPartitionCollection.cpp b/Engine/Source/Runtime/Experimental/ChaosSpatialPartitions/Private/ChaosSpatialPartitions/Collections/DefaultSpatialPartitionCollection.cpp
--- a/Engine/Source/Runtime/Experimental/ChaosSpatialPartitions/Private/ChaosSpatialPartitions/Collections/DefaultSpatialPartitionCollection.cpp
+++ b/Engine/Source/Runtime/Experimental/ChaosSpatialPartitions/Private/ChaosSpatialPartitions/Collections/DefaultSpatialPartitionCollection.cpp
@@ -72,7 +72,7 @@ namespace Chaos::SpatialPartition
 		// Changes are not valid during a rebuild
 		check(!IsRebuilding());
 
-		const int32 HandleValue = (int32)InOutHandle.GetValue();
+		const uint32 HandleValue = (uint32)InOutHandle.GetValue();
 		InOutHandle.SetValue(INDEX_NONE);
 		FEntryHandle EntryHandle;
 		EntryHandle.FromUint(HandleValue);
diff --git a/Engine/Source/Runtime/Experimental/ChaosSpatialPartitions/Private/ChaosSpatialPartitions/Library/StaticAabbTree.cpp b/Engine/Source/Runtime/Experimental/ChaosSpatialPartitions/Private/ChaosSpatialPartitions/Library/StaticAabbTree.cpp
--- a/Engine/Source/Runtime/Experimental/ChaosSpatialPartitions/Private/ChaosSpatialPartitions/Library/StaticAabbTree.cpp
+++ b/Engine/Source/Runtime/Experimental/ChaosSpatialPartitions/Private/ChaosSpatialPartitions/Library/StaticAabbTree.cpp
@@ -65,12 +65,15 @@ namespace Chaos::SpatialPartition
 
 		FEntry& Entry = Entries[EntryIndex];
 
-		const FAabbTreeLeaf& Leaf = Leaves[Entry.LeafIndex];
-		if (Leaf.Aabb.Contains(Aabb))
+		if (Leaves.IsValidIndex(Entry.LeafIndex))
 		{
-			Entry.Aabb = Aabb;
-			Entry.UserData = UserData;
-			return true;
+			const FAabbTreeLeaf& Leaf = Leaves[Entry.LeafIndex];
+			if (Leaf.Aabb.Contains(Aabb))
+			{
+				Entry.Aabb = Aabb;
+				Entry.UserData = UserData;
+				return true;
+			}
 		}
 		return false;
 	}
@@ -236,7 +239,7 @@ namespace Chaos::SpatialPartition
 		if (LeavesFreeListHead != INDEX_NONE)
 		{
 			const int32 Index = LeavesFreeListHead;
-			EntriesFreeListHead = Leaves[Index].NodeIndex;
+			LeavesFreeListHead = Leaves[Index].NodeIndex;
 
 			Leaves[Index] = FAabbTreeLeaf();
 			return Index;
@@ -331,9 +334,13 @@ namespace Chaos::SpatialPartition
 	void FStaticAabbTree::RemoveEntryFromLeaf(const int32 EntryIndex)
 	{
 		FEntry& Entry = Entries[EntryIndex];
-		FAabbTreeLeaf& LeafNode = Leaves[Entry.LeafIndex];
-		LeafNode.Elements.RemoveAllSwap([EntryIndex](const FAabbTreeLeafElement& Element) {return Element.Index == EntryIndex; });
-		Entry.LeafIndex = INDEX_NONE;
+		// Leaf might never be set if the object was inserted and the tree was never rebuilt
+		if (Entry.LeafIndex != INDEX_NONE)
+		{
+			FAabbTreeLeaf& LeafNode = Leaves[Entry.LeafIndex];
+			LeafNode.Elements.RemoveAllSwap([EntryIndex](const FAabbTreeLeafElement& Element) {return Element.Index == EntryIndex; });
+			Entry.LeafIndex = INDEX_NONE;
+		}
 	}
 
 	template <typename QueryDataType, typename VisitorType>
diff --git a/Engine/Source/Runtime/Experimental/ChaosSpatialPartitions/Tests/ChaosSpatialPartitions/Algorithms/PartitioningAlgorithmsTests.cpp b/Engine/Source/Runtime/Experimental/ChaosSpatialPartitions/Tests/ChaosSpatialPartitions/Algorithms/PartitioningAlgorithmsTests.cpp
--- a/Engine/Source/Runtime/Experimental/ChaosSpatialPartitions/Tests/ChaosSpatialPartitions/Algorithms/PartitioningAlgorithmsTests.cpp
+++ b/Engine/Source/Runtime/Experimental/ChaosSpatialPartitions/Tests/ChaosSpatialPartitions/Algorithms/PartitioningAlgorithmsTests.cpp
@@ -35,6 +35,7 @@ namespace Chaos::SpatialPartition::LowLevelTest
 		InputToOutputIndices.SetNum(Objects.Num());
 		for (int32 I = 0; I < Objects.Num(); ++I)
 		{
+			REQUIRE(InputToOutputIndices.IsValidIndex(Objects[I].UserData));
 			InputToOutputIndices[Objects[I].UserData] = I;
 		}
 
@@ -316,10 +317,15 @@ namespace Chaos::SpatialPartition::LowLevelTest
 			{
 				FVec3 Spacing(1);
 				Spacing[Axis] = 2;
+				BuildObjectList(Objects, 3, 3, 3, Spacing);
+				
 				// Make every other axis larger. This shouldn't affect anything as only the center is considered.
 				FVec3 AabbSize(5);
 				AabbSize[Axis] = 1;
-				BuildObjectList(Objects, 3, 3, 3, Spacing);
+				for (FTestObject& Object : Objects)
+				{
+					Object.Aabb = BuildAabbCenterExtents(Object.Aabb.GetCenter(), AabbSize);
+				}
 
 				TArrayView<const FTestObject> View(Objects);
 				AabbTreeAlgorithm::ComputeSplitPlaneWithSpatialMedianHeuristic(View, SplitAxis, SplitPosition);
